CQTest - Update instalation section (doc)

diff --git a/Engine/Plugins/Tests/CQTest/README.md b/Engine/Plugins/Tests/CQTest/README.md
--- a/Engine/Plugins/Tests/CQTest/README.md
+++ b/Engine/Plugins/Tests/CQTest/README.md
@@ -75,7 +75,6 @@ The inspiration for CQTest was to add the before/after test abilities, while res
 ```
 
 # Installation
-Clone the project, into your project's Plugin folder.
 Inside the project that you want to test, you'll need to change 2 things:
 In the .uproject file of the project you want to test, add the following to the Plugins section
 
