[doc] Remove email reference in readme file

diff --git a/Engine/Plugins/Tests/CQTest/README.md b/Engine/Plugins/Tests/CQTest/README.md
--- a/Engine/Plugins/Tests/CQTest/README.md
+++ b/Engine/Plugins/Tests/CQTest/README.md
@@ -359,4 +359,3 @@ And creating a macro which uses it
 
 # Contribute
 Improvements like bug fixes and extensions are welcome when accompanied by unit tests.
-Feel free to reach out to [Dustin Randall](mailto:[REDACTED_EMAIL]) with questions or comments
