Updating CQTest readme and documentation

diff --git a/Engine/Plugins/Tests/CQTest/README.md b/Engine/Plugins/Tests/CQTest/README.md
--- a/Engine/Plugins/Tests/CQTest/README.md
+++ b/Engine/Plugins/Tests/CQTest/README.md
@@ -1,8 +1,10 @@
-# Introduction 
+# Introduction
+
 Extension of the Unreal Engine FAutomationTestBase to provide test fixtures and common automation testing commands.
 
 # Why CQTest?
-There are other valid ways of testing in Unreal engine.  One option is to use the provided macros from Unreal Engine: [docs](https://docs.unrealengine.com/5.1/en-US/automation-technical-guide/)
+
+There are other valid ways of testing in Unreal engine.  One option is to use the provided macros from Unreal Engine: [docs](https://dev.epicgames.com/documentation/en-us/unreal-engine/write-cplusplus-tests-in-unreal-engine)
 ```cpp
     IMPLEMENT_SIMPLE_AUTOMATION_TEST(FMinimalTest, "Game.Test", EAutomationTestFlags::EditorContext | EAutomationTestFlags::EngineFilter)
 	bool FMinmalTest::RunTest(const FString& Parameters) 
@@ -12,7 +14,7 @@ There are other valid ways of testing in Unreal engine.  One option is to use th
 	}
 ```
 
-Unreal has also developed their [spec test framework](https://docs.unrealengine.com/4.27/en-US/TestingAndOptimization/Automation/AutomationSpec/), which is inspired by Behavior Driven Design
+Unreal Engine also has a [spec test framework](https://dev.epicgames.com/documentation/en-us/unreal-engine/automation-spec-in-unreal-engine), which is inspired by [Behavior Driven Design](https://en.wikipedia.org/wiki/Behavior-driven_development)
 ```cpp
     DEFINE_SPEC(FMinimalTest, "Game.Test", EAutomationTestFlags::ProductFilter | EAutomationTestFlags_ApplicationContextMask)
 	void FMinimalTest::Define() 
@@ -26,8 +28,8 @@ Unreal has also developed their [spec test framework](https://docs.unrealengine.
 		});
 	}
 ```
-With the spec tests, be careful about capturing state
 
+With the spec tests, be careful about capturing state
 ```cpp
     BEGIN_DEFINE_SPEC(FMinimalTest, "Game.Test", EAutomationTestFlags::ProductFilter | EAutomationTestFlags_ApplicationContextMask)
 		uint32 SomeValue = 3;
@@ -51,8 +53,7 @@ With the spec tests, be careful about capturing state
 	}
 ```
 
-The inspiration for CQTest was to add the before/after test abilities, while resetting state between tests automatically.  One of the guiding principles is to make easy things easy.
-
+The inspiration for **CQTest** was to add the before/after test abilities, while resetting state between tests automatically.  One of the guiding principles is to make easy things easy.
 ```cpp
 	TEST(MinimalTest, "Game.Test") 
 	{
@@ -75,9 +76,9 @@ The inspiration for CQTest was to add the before/after test abilities, while res
 ```
 
 # Installation
+
 Inside the project that you want to test, you'll need to change 2 things:
 In the .uproject file of the project you want to test, add the following to the Plugins section
-
 ```json
 		{
 			"Name": "CQTest",
@@ -85,8 +86,7 @@ In the .uproject file of the project you want to test, add the following to the
 		}
 ```
 
-Then in the project's .Build.cs file, you'll want to add the following to the PrivateDependencyModuleNames.  Something like
-
+Then in the project's .Build.cs file, you'll want to add the following to the `PrivateDependencyModuleNames`.  Something like
 ```csharp
 		PrivateDependencyModuleNames.AddRange(
 			new string[] {
@@ -98,22 +98,22 @@ Then in the project's .Build.cs file, you'll want to add the following to the Pr
 			);
 ```
 
-	
 # Build
+
 No special steps are required to build the plugin project, it should build with the rest of the project.
 
 # Test
+
 This plugin has a set of tests to validate and document the behavior.  To run tests in Unreal
 - Launch the editor
 - Find the Tools drop down and select Session Frontend
 - Navigate to the Automation tab
-- By default, the tests should be listed first under " Product.Plugins.CQTest"
+- By default, the tests should be listed first under "Product.Plugins.CQTest"
 - Select the tests you would like to run and press 'Start Tests'
 
 # Examples
 
 Tests can be as simple as
-
 ```cpp
     #include "CQTest.h"
 	
@@ -123,8 +123,7 @@ Tests can be as simple as
 	}
 ```
 
-For setup and teardown, or common state between multiple tests, or to group related tests, use the TEST_CLASS macro.
-
+For setup and teardown, or common state between multiple tests, or to group related tests, use the `TEST_CLASS` macro.
 ```cpp
     #include "CQTest.h"
 	TEST_CLASS(MyNeatTest, "Game.MyGame") 
@@ -133,6 +132,13 @@ For setup and teardown, or common state between multiple tests, or to group rela
 		uint32 SomeNumber = 0;
 		Thing* Thing = nullptr;
 		
+		// Optional static method to be executed before all tests of this TEST_CLASS
+		// Should be removed if empty and unused
+		BEFORE_ALL() 
+		{
+			// Perform some logic that is shared with all tests such as loading a level
+		}
+
 		BEFORE_EACH() 
 		{
 			SetupCalled = true;
@@ -142,7 +148,14 @@ For setup and teardown, or common state between multiple tests, or to group rela
 		
 		AFTER_EACH() 
 		{
-			delete Thing; //Should normally use RAII for things like this
+			delete Thing; // Should normally use RAII for things like this
+		}
+
+		// Optional static method to be executed after all tests of this TEST_CLASS
+		// Should be removed if empty and unused
+		AFTER_ALL()
+		{
+			// Perform cleanup of any resources that was done in the `BEFORE_ALL`.
 		}
 		
 	protected:
@@ -168,8 +181,14 @@ For setup and teardown, or common state between multiple tests, or to group rela
 	};
 ```
 
-Test Directory determines where in the Automation tab the tests appear.  In the example above, we specify "Game.MyGame", but you may also have an auto-generated test directory based on the folder structure.
+In addition to `TEST` and `TEST_CLASS` are 5 additional macros:
+ - `TEST_CLASS_WITH_ASSERTS` - Macro which allows this test object to use a custom asserter.  More information about how to use this macro can be found [in the section below regarding assertions](#assertions)
+ - `TEST_CLASS_WITH_BASE` - Macro which allows this test object to inherit from a different test object.  More information about how to use this macro can be found [in the section below regarding custom test classes](#base-test-class)
+ - `TEST_CLASS_WITH_FLAGS` - Macro which allows the use of different automation test flags to be specified.  Useful for when tests can only run under a certain context or grouped under a specific filter.  The default flags for `TEST` and `TEST_CLASS` are `EAutomationTestFlags_ApplicationContextMask | EAutomationTestFlags::ProductFilter`.  Additional information regarding the available flags can be found [in the online documentation](https://dev.epicgames.com/documentation/en-us/unreal-engine/API/Runtime/Core/Misc/EAutomationTestFlags/Type)
+ - `TEST_CLASS_WITH_BASE_AND_FLAGS` - Macro which allows for this test object to inherit from a different test object and allows for custom automation test flags to be specified.
+ - `TEST_CLASS_IMPL` - Base macro which is used by the above macros to specify a custom asserter, a test object to inherit from a different test object, and allows for custom automation test flags to be specified.
 
+Test Directory determines where in the Automation tab the tests appear.  In the example above, we specify "Game.MyGame", but you may also have an auto-generated test directory based on the folder structure.
 ```cpp
 	TEST_CLASS(MyNeatTest, GenerateTestDirectory)
 	{
@@ -179,9 +198,9 @@ Test Directory determines where in the Automation tab the tests appear.  In the
 	{
 	};
 ```
+In the above examples, if `MyNeatTest` is located within a plugin of a project with the path **MyProject/Plugins/GameTests/Source/GameTests/Private/NeatTest.cpp** then the generated test names will be `MyProject.Plugins.GameTests.MyNeatTest` and `Game.Test.MyProject.Plugins.GameTests.MyNeatTest.Validation` respectively.
 
 Constructors (and destructors) are available.  Destructors shouldn't throw, and you shouldn't put assertions in them (as they are called after the testing framework is done with the test).
-
 ```cpp
 	TEST_CLASS(SomeTestClass, "Game.Test")
 	{
@@ -198,8 +217,7 @@ Constructors (and destructors) are available.  Destructors shouldn't throw, and
 	};
 ```	
 
-Latent actions are supported with the TEST_CLASS macro.  Each step will complete all latent actions before moving to the next.  If an assertion is raised during a latent action, then no further latent actions will be processed.  The AFTER_EACH method will still be invoked though.
-
+Latent actions are supported with the `TEST_CLASS` macro.  Each step will complete all latent actions before moving to the next.  If an assertion is raised during a latent action, then no further latent actions will be processed.  The `AFTER_EACH` method will still be invoked though.
 ```cpp
     TEST_CLASS(LatentActionTest, "Game.Test") 
 	{
@@ -223,8 +241,14 @@ Latent actions are supported with the TEST_CLASS macro.  Each step will complete
 	};
 ```
 
-Also available for commands is a fluent command builder
+**CQTest** provides the following additional latent actions:
+ - `FExecute` - Action that executes only once.
+ - `FWaitUntil` - Action that executes over multiple ticks until either completion or the duration exceeds the timeout.  Action will fail if the condition cannot be satisifed before timing out.
+ - `FWaitDelay` - Action that waits a specified duration.
+   - **CAUTION:** Using a timed-wait can introduce flakiness due to variable runtimes and the above `FWaitUntil` should be used instead.
+ - `FRunSequence` - Action which ensures that a collection of latent actions occur in order, and only after all previous actions have finished.
 
+Also available for commands is a fluent command builder
 ```cpp
 	TEST_METHOD(SomeTest) 
 	{
@@ -236,35 +260,52 @@ Also available for commands is a fluent command builder
 	}
 ```
 
+The command builder provides commands which wrap around the above mentioned latent actions.  The following commands are made available:
+ - `Do`/`Then` - Commands which adds the `FExecute` latent action with the provided lambda to be executed.
+ - `StartWhen`/`Until` - Commands which adds the `FWaitUntil` latent action with the provided lambda to be evaluated.
+ - `WaitDelay` - Command which waits a specified duration before continuing.  
+   - **CAUTION:** Using a timed-wait can introduce flakiness due to variable runtimes and the above `StartWhen`/`Until` commands should be used instead.
+ - `OnTearDown`/`CleanUpWith` - Commands which adds the `FExecute` latent action with the provided lambda to be executed after the test.  Can be called multiple times to add multiple clean up latent actions.
+   - **NOTE:** Latent actions added using the `OnTearDown`/`CleanUpWith` will be run in reverse order (i.e. Last in, first out)
+
 The framework will ensure that all of those commands happen in order using a future pattern.
-Similarly, the framework will ensure that a test can await a ticking object.  See GameObjectsTickTest for an example
-One word of caution, the framework does not currently support adding latent actions from within latent actions.
+Similarly, the framework will ensure that a test can await a ticking object.  See `GameObjectsTickTest` for an example
+
+**CAUTION:** The framework does not currently support adding latent actions from within latent actions.
 Instead, it is better to add the actions as a series of self-contained steps.
- 
 
 # Extending the framework
-The framework has been designed to allow for extensions in a couple areas.  See ExtensionTests.cpp for in-code examples.
+
+The framework has been designed to allow for extensions in a couple areas.  See _CQTestTests/Private/ExtensionTests.cpp_ for in-code examples.
 
 ## Test Components
+
 This testing framework embraces composition over inheritence.  Creating new components should be the default mechanism for extending the framework.  Some of the components available to you are:
-  ActorTestSpawner - Allows a test to spawn actors, and manages their despawning.
-  MapTestSpawner - Creates a map and opens a level.  Allows tests to spawn actors in that world.
-  BlueprintHelper - Eases the ability for a test to spawn Blueprint objects, intended to be used with MapTestSpawner.
-  PIENetworkComponent - Allows tests to create a server and a collection of clients.  Good for testing replication.
+ - `SpawnHelper` - Eases the ability to spawn actors and other objects.  Implemented by `ActorTestSpawner` and `MapTestSpawner`.
+ - `ActorTestSpawner` - Creates a minimal `UWorld` for a test to spawn actors, and manages their despawning.
+ - `MapTestSpawner` - Can create a temporary map or open a specified level.  Allows tests to spawn actors in that world.
+ - `CQTestBlueprintHelper` - Eases the ability for a test to spawn Blueprint objects, intended to be used with `MapTestSpawner`.
+   - **NOTE:** Loading Blueprint assets is only intended to work within the Editor context.  Tests that make use of the `CQTestBlueprintHelper` should specify the `EAutomationTestFlags::EditorContext` flag.
+ - `PIENetworkComponent` - Allows tests to create a server and a collection of clients.  Good for testing replication.
+    - **NOTE:** The `PIENetworkComponent` sets up a Server and Client PIE instance which is only usable within the Editor context.  Tests that make use of the `PIENetworkComponent` should specify the `EAutomationTestFlags::EditorContext` flag.
+ - `InputTestActions` - Allows tests to inject `InputActions` to the `Pawn`.
+ - `CQTestSlateComponent` - Allows tests to get notified when the UI has been updated.
 
 ## Assertions
+
 Not all platforms support exceptions, and so the assertions are unable to rely on them.
 There are a few options here:
-  We could just throw exceptions, and only run tests on platforms which support exceptions
-  We could return a [[nodiscard]] bool to encourage checking each assertion and returning if it fails
-  We could return a normal bool and rely on people to check it when it's important.
+ - We could just throw exceptions, and only run tests on platforms which support exceptions
+ - We could return a `[[nodiscard]]` bool to encourage checking each assertion and returning if it fails
+ - We could return a normal bool and rely on people to check it when it's important.
+
 Exceptions have the advantage of working in helper functions and lambdas, as well as not depending on human diligence.
 A normal bool is less noisy, and allows developers to use intellisense, but is more error prone
-The default implementation used is the [[nodiscard]] bool, with a helper macro ASSERT_THAT which does the early return check for you.
+The default implementation used is the `[[nodiscard]]` bool, with a helper macro `ASSERT_THAT` which does the early return check for you.
 
-You can use your own types within the Assert.AreEqual and Assert.AreNotEqual methods assuming you have the == and != operators defined as needed.
-In addition, the error message will print out the string version of your type, assuming you have a ToString method defined as well.  The framework will complain if it doesn't know how to print your value.
-You can find examples of providing a string to the framework in CQTestConvertTests.cpp, but below is a simple example.
+You can use your own types within the `Assert.AreEqual` and `Assert.AreNotEqual` methods assuming you have the `==` and `!=` operators defined as needed.
+In addition, the error message will print out the string version of your type, assuming you have a `ToString` method defined as well.  The framework will complain if it doesn't know how to print your value.
+You can find examples of providing a string to the framework in _CQTestTests/Private/Assert/CQTestConvertTests.cpp_, but below is a simple example.
 ```cpp
 struct MyCustomType
 {
@@ -297,7 +338,7 @@ FString CQTestConvert::ToString(const MyCustomEnum&)
 ```
 
 You are able to customize the assertions which are available, and how they behave.
-See CQTestTests/Private/ExtensionTests.cpp for an example
+See _CQTestTests/Private/ExtensionTests.cpp_ for an example
 Below is some untested example code to inspire ideas
 
 ```cpp
@@ -340,7 +381,6 @@ Below is some untested example code to inspire ideas
 ```
 
 From here, you could create macros your studio uses to create tests
-
 ```cpp
 	#define MY_STUDIO_TEST_CLASS(_ClassName, _TestDir) TEST_CLASS_WITH_ASSERTS(_ClassName, _TestDir, FluentAsserter)
 	#define MY_STUDIO_TEST(_TestName, _TestDir) \
@@ -352,8 +392,8 @@ From here, you could create macros your studio uses to create tests
 ```
 
 ## Base test class
-Similarly there may be a use case to create many tests which have the same member variables or helper methods.  This can be implemented by extending the test class
 
+Similarly there may be a use case to create many tests which have the same member variables or helper methods.  This can be implemented by extending the test class
 ```cpp
 	template<typename Derived, typename AsserterType>
 	struct ActorTest : public Test<Derived, AsserterType>
@@ -363,10 +403,10 @@ Similarly there may be a use case to create many tests which have the same membe
 ```
 
 And creating a macro which uses it
-
 ```cpp
 	#define ACTOR_TEST(_ClassName, _TestDir) TEST_CLASS_WITH_BASE(_ClassName, _TestDir, ActorTest)
 ```
 
 # Contribute
+
 Improvements like bug fixes and extensions are welcome when accompanied by unit tests.
diff --git a/Engine/Plugins/Tests/CQTest/Source/CQTest/Public/Commands/TestCommands.h b/Engine/Plugins/Tests/CQTest/Source/CQTest/Public/Commands/TestCommands.h
--- a/Engine/Plugins/Tests/CQTest/Source/CQTest/Public/Commands/TestCommands.h
+++ b/Engine/Plugins/Tests/CQTest/Source/CQTest/Public/Commands/TestCommands.h
@@ -4,6 +4,7 @@
 
 #include "Misc/AutomationTest.h"
 
+/** Latent Command that waits until the Query evaluates to `true` or the timeout has exceeded. */
 class CQTEST_API FWaitUntil : public IAutomationLatentCommand
 {
 public:
@@ -51,6 +52,7 @@ enum class ECQTestFailureBehavior
 	Run
 };
 
+/** Latent Command which executes the provided function. */
 class CQTEST_API FExecute : public IAutomationLatentCommand
 {
 public:
@@ -69,6 +71,7 @@ public:
 	ECQTestFailureBehavior FailureBehavior;
 };
 
+/** Latent Command which manages and executes an array of latent commands. */
 class CQTEST_API FRunSequence : public IAutomationLatentCommand
 {
 public:
diff --git a/Engine/Plugins/Tests/CQTest/Source/CQTest/Public/Components/ActorTestSpawner.h b/Engine/Plugins/Tests/CQTest/Source/CQTest/Public/Components/ActorTestSpawner.h
--- a/Engine/Plugins/Tests/CQTest/Source/CQTest/Public/Components/ActorTestSpawner.h
+++ b/Engine/Plugins/Tests/CQTest/Source/CQTest/Public/Components/ActorTestSpawner.h
@@ -4,19 +4,47 @@
 
 #include "SpawnHelper.h"
 
+/*
+//Example boiler plate
+
+#include "CQTest.h"
+#include "Components/ActorTestSpawner.h"
+
+TEST_CLASS(MyFixtureName, "ActorSpawner.Example")
+{
+	FActorTestSpawner Spawner;
+
+	TEST_METHOD(Spawn_BaseActor_DefaultActorSpawned)
+	{
+		AActor& Actor = Spawner.SpawnActor<AActor>();
+		ASSERT_THAT(IsTrue(Actor.GetFName() == NAME_None));
+	}
+
+	TEST_METHOD(Spawn_BaseObject_DefaultObjectSpawned)
+	{
+		UObject& Object = Spawner.SpawnObject<UObject>();
+		ASSERT_THAT(IsTrue(Object.GetFName() == NAME_None));
+	}
+};
+*/
+
 class UTestGameInstance;
 
-/// Class for spawning actors in an ActorTest context (no PIE loaded)
+/** Class for spawning actors in an ActorTest context(no PIE loaded) */
 struct CQTEST_API FActorTestSpawner : public FSpawnHelper
 {
 	FActorTestSpawner() = default;
+
 	virtual ~FActorTestSpawner() override;
 
+	/** Initialize the GameSubsystem for the test world. */
 	void InitializeGameSubsystems();
 
+	/** Returns a pointer to the test game instance. */
 	UTestGameInstance* GetGameInstance();
 
 protected:
+	/** Returns a newly created world. */
 	virtual UWorld* CreateWorld() override;
 
 private:
diff --git a/Engine/Plugins/Tests/CQTest/Source/CQTest/Public/Components/CQTestBlueprintHelper.h b/Engine/Plugins/Tests/CQTest/Source/CQTest/Public/Components/CQTestBlueprintHelper.h
--- a/Engine/Plugins/Tests/CQTest/Source/CQTest/Public/Components/CQTestBlueprintHelper.h
+++ b/Engine/Plugins/Tests/CQTest/Source/CQTest/Public/Components/CQTestBlueprintHelper.h
@@ -4,9 +4,55 @@
 #include "UObject/Object.h"
 #include "UObject/Class.h"
 
-/// Utilities for working with Blueprint in C++ Tests
+/*
+//Example boiler plate
+
+#if WITH_AUTOMATION_TESTS
+
+#include "CQTest.h"
+#include "Components/CQTestBlueprintHelper.h"
+
+TEST_CLASS_WITH_FLAGS(MyFixtureName, "BlueprintHelper.Example", EAutomationTestFlags::EditorContext | EAutomationTestFlags::ProductFilter)
+{
+	FCQTestBlueprintHelper BlueprintHelper;
+
+	TEST_METHOD(Object_FromBlueprint_IsLoaded)
+	{
+		UClass* Class = BlueprintHelper.GetBlueprintClass(TEXT("/Package/Path/To/Asset"), TEXT("B_Asset"));
+		ASSERT_THAT(IsNotNull(Class));
+
+		UObject* Object = BlueprintHelper.FindDataBlueprint(TEXT("/Package/Path/To/Asset"), TEXT("B_Asset"));
+		ASSERT_THAT(IsNotNull(Object));
+	}
+};
+
+#endif // WITH_AUTOMATION_TESTS
+*/
+
+/** Utilities for working with Blueprint in C++ Tests */
 struct CQTEST_API FCQTestBlueprintHelper
 {
+	/**
+	 * Finds the UObject given a directory and asset name to find.
+	 *
+	 * @param Directory - Directory path where the asset resides.
+	 * @param Name - Name of the asset.
+	 * 
+	 * @return pointer to the UObject that was found and loaded, returns nullptr if no object was found
+	 * 
+	 * @note Method uses a static cache to store information regarding directories and assets.
+	 */
 	UObject* FindDataBlueprint(const FString& Directory, const FString& Name);
+
+	/**
+	 * Gets the UClass of a UObject loaded from a given directory and asset name.
+	 *
+	 * @param Directory - Directory path where the asset resides.
+	 * @param Name - Name of the asset.
+	 * 
+	 * @return pointer of the blueprint's generated class
+	 *
+	 * @note Method guarantees that the UClass returned is valid, will assert otherwise.
+	 */
 	UClass* GetBlueprintClass(const FString& Directory, const FString& Name);
 };
\ No newline at end of file
diff --git a/Engine/Plugins/Tests/CQTest/Source/CQTest/Public/Components/CQTestSlateComponent.h b/Engine/Plugins/Tests/CQTest/Source/CQTest/Public/Components/CQTestSlateComponent.h
--- a/Engine/Plugins/Tests/CQTest/Source/CQTest/Public/Components/CQTestSlateComponent.h
+++ b/Engine/Plugins/Tests/CQTest/Source/CQTest/Public/Components/CQTestSlateComponent.h
@@ -8,7 +8,31 @@
 
 #include <atomic>
 
-// CQTest component for interacting with Slate
+/*
+//Example boiler plate
+
+#include "CQTest.h"
+#include "Components/CQTestSlateComponent.h"
+
+TEST_CLASS(MyFixtureName, "Slate.Example")
+{
+	TUniquePtr<FCQTestSlateComponent> SlateComponent;
+
+	BEFORE_EACH()
+	{
+		SlateComponent = MakeUnique<FCQTestSlateComponent>();
+	}
+
+	TEST_METHOD(HaveTicksElapsed_WaitUntil_EventuallyReturnsTrue)
+	{
+		TestCommandBuilder
+			.StartWhen([this]() { return SlateComponent->HaveTicksElapsed(3); })
+			.Then([this]() { ASSERT_THAT(IsTrue(SlateComponent->GetTotalElapsedTicks() >= 3)); });
+	}
+};
+*/
+
+/** CQTest component for interacting with Slate */
 class CQTEST_API FCQTestSlateComponent
 {
 public:
@@ -18,11 +42,12 @@ public:
 	/**
 	 * Initializes the ExpectedTick variable to be checked against the amount of times Slate has ticked.
 	 * Will be used to check Slate has ticked the expected amount before confirming and resetting the ExpectedTick.
-	 * 
-	 * Note that this method is meant to be used as the only statement within a latent command that waits until the predicate is met, such as FWaitUntil, Until, or StartWhen methods.
 	 *
 	 * @param Ticks - Number of ticks we're expecting to have elapsed
+	 * 
 	 * @return true when Slate has ticked the provided number of times.
+	 * 
+	 * @note Method is meant to be used as the only statement within a latent command that waits until the predicate is met, such as FWaitUntil, Until, or StartWhen methods.
 	 */
 	bool HaveTicksElapsed(uint32 Ticks);
 
diff --git a/Engine/Plugins/Tests/CQTest/Source/CQTest/Public/Components/InputTestActions.h b/Engine/Plugins/Tests/CQTest/Source/CQTest/Public/Components/InputTestActions.h
--- a/Engine/Plugins/Tests/CQTest/Source/CQTest/Public/Components/InputTestActions.h
+++ b/Engine/Plugins/Tests/CQTest/Source/CQTest/Public/Components/InputTestActions.h
@@ -6,9 +6,94 @@
 #include "InputAction.h"
 #include "InputActionValue.h"
 
+/*
+//Example boiler plate
+
+#include "CQTest.h"
+#include "Components/InputTestActions.h"
+#include "Components/MapTestSpawner.h"
+
+// It is recommended that all code before the TEST_CLASS should reside in a separate file that is shared by multiple input tests
+struct FInputTestAction : public FTestAction
+{
+	FInputTestAction(const FInputActionValue& InInputActionValue)
+	{
+		// Name of the InputAction assigned to the Pawn
+		InputActionName = TEXT("InputActionName");
+
+		// Value assigned to the input. Can be boolean for button presses or a FVector for axis direction
+		InputActionValue = InInputActionValue;
+	}
+};
+
+// Create test actions which inherits from our input action, but provides an input vector along our movement axis
+struct FAxisInputAction : public FInputTestAction
+{
+	FAxisInputAction() : FInputTestAction(FVector2D(0.0f, 1.0f))
+	{
+	}
+};
+
+class FPawnTestActions : public FInputTestActions
+{
+public:
+	explicit FPawnTestActions(APawn* Pawn) : FInputTestActions(Pawn)
+	{
+	}
+
+	// Simulates player movement input actions.
+	void PerformMovement()
+	{
+		// Perform move actions over the duration of 1 second
+		PerformAction(FAxisInputAction{}, [this]() -> bool {
+			if (StartTime.GetTicks() == 0)
+			{
+				StartTime = FDateTime::UtcNow();
+			}
+
+			FTimespan Elapsed = FDateTime::UtcNow() - StartTime;
+			return Elapsed >= FTimespan::FromSeconds(1.0);
+		});
+	}
+
+	FDateTime StartTime{ 0 };
+};
+
+TEST_CLASS(MyFixtureName, "InputActions.Example")
+{
+	TUniquePtr<FMapTestSpawner> Spawner;
+	TUniquePtr<FPawnTestActions> PawnActions;
+	APawn* Player;
+
+	bool IsPlayerMoving()
+	{
+		return FMath::IsNearlyEqual(Player->GetVelocity().Length(), 0.0, UE_KINDA_SMALL_NUMBER);
+	}
+
+	BEFORE_EACH()
+	{
+		Spawner = MakeUnique<FMapTestSpawner>(TEXT("/Package/Path/To/Map"), TEXT("MapName"));
+		Spawner->AddWaitUntilLoadedCommand(TestRunner);
+	}
+
+	TEST_METHOD(PlayerStartsMoving_ForDuration_EventuallyStops)
+	{
+		TestCommandBuilder
+			.StartWhen([this]() { return nullptr != Spawner->FindFirstPlayerPawn(); })
+			.Then([this]() {
+				Player = Spawner->FindFirstPlayerPawn();
+				PawnActions = MakeUnique<FPawnTestActions>(Player);
+			})
+			.Then([this]() { PawnActions->PerformMovement(); })
+			.Then([this]() { ASSERT_THAT(IsTrue(IsPlayerMoving())); })
+			.Until([this]() { return !IsPlayerMoving(); });
+	}
+};
+*/
+
 class APawn;
 
-// Class for testing input of a Pawn by injecting InputActions
+/** Class for testing input of a Pawn by injecting InputActions */
 class CQTEST_API FTestAction
 {
 public:
@@ -37,10 +122,15 @@ private:
 	const UInputAction* InputAction{ nullptr };
 };
 
-// Class for processing FTestAction objects
+/** Class for processing FTestAction objects */
 class CQTEST_API FInputTestActions
 {
 public:
+	/**
+	 * Construct the InputTestActions.
+	 *
+	 * @param InPawn - Pawn that will have the InputActions applied to.
+	 */
 	FInputTestActions(APawn* InPawn) : Pawn(InPawn)
 	{
 	}
@@ -72,6 +162,7 @@ private:
 	 * Processes repeat actions every tick.
 	 *
 	 * @param DeltaTime - Time elapsed since the last tick.
+	 * 
 	 * @return true if tick has completed.
 	 */
 	bool Tick(float DeltaTime);
diff --git a/Engine/Plugins/Tests/CQTest/Source/CQTest/Public/Components/MapTestSpawner.h b/Engine/Plugins/Tests/CQTest/Source/CQTest/Public/Components/MapTestSpawner.h
--- a/Engine/Plugins/Tests/CQTest/Source/CQTest/Public/Components/MapTestSpawner.h
+++ b/Engine/Plugins/Tests/CQTest/Source/CQTest/Public/Components/MapTestSpawner.h
@@ -6,13 +6,46 @@
 
 #include "Commands/TestCommandBuilder.h"
 
+/*
+//Example boiler plate
+
+#include "CQTest.h"
+#include "Components/MapTestSpawner.h"
+
+#if WITH_AUTOMATION_TESTS
+
+TEST_CLASS(MyFixtureName, "MapSpawner.Example")
+{
+	TUniquePtr<FMapTestSpawner> Spawner;
+	APawn* MyPawn;
+
+	BEFORE_EACH()
+	{
+		Spawner = MakeUnique<FMapTestSpawner>(TEXT("/Package/Path/To/Map"), TEXT("MapName"));
+		Spawner->AddWaitUntilLoadedCommand(TestRunner);
+	}
+
+	TEST_METHOD(PlayerPawn_Loaded_Found)
+	{
+		TestCommandBuilder
+			.StartWhen([this]() {
+				MyPawn = Spawner->FindFirstPlayerPawn();
+				return nullptr != MyPawn;
+			})
+			.Then([this]() { ASSERT_THAT(IsNotNull(MyPawn)); });
+	}
+};
+
+#endif // WITH_AUTOMATION_TESTS
+*/
+
 #if WITH_AUTOMATION_TESTS
 
 #if WITH_EDITOR
 #include "UnrealEdMisc.h"
 #endif // WITH_EDITOR
 
-/// Class for spawning Actors in a named Map/Level
+/** Class for spawning Actors in a named Map / Level */
 struct CQTEST_API FMapTestSpawner : public FSpawnHelper
 {
 	/**
@@ -28,6 +61,7 @@ struct CQTEST_API FMapTestSpawner : public FSpawnHelper
 	 * Creates an instance of the MapTestSpawner with a temporary level ready for use.
 	 * 
 	 * @param InCommandBuilder - Test Command Builder used to assist with setup.
+	 * 
 	 * @return unique instance of the FMapTestSpawner, nullptr otherwise
 	 */
 	static TUniquePtr<FMapTestSpawner> CreateFromTempLevel(FTestCommandBuilder& InCommandBuilder);
@@ -45,6 +79,7 @@ struct CQTEST_API FMapTestSpawner : public FSpawnHelper
 	APawn* FindFirstPlayerPawn();
 
 protected:
+	/** Returns a newly created world. */
     virtual UWorld* CreateWorld() override;
 
 private:
diff --git a/Engine/Plugins/Tests/CQTest/Source/CQTest/Public/Components/PIENetworkComponent.h b/Engine/Plugins/Tests/CQTest/Source/CQTest/Public/Components/PIENetworkComponent.h
--- a/Engine/Plugins/Tests/CQTest/Source/CQTest/Public/Components/PIENetworkComponent.h
+++ b/Engine/Plugins/Tests/CQTest/Source/CQTest/Public/Components/PIENetworkComponent.h
@@ -13,37 +13,204 @@
 #include "Engine/NetDriver.h"
 #include "Engine/PackageMapClient.h"
 
+/*
+//Example boiler plate
+
+#include "CQTest.h"
+#include "Components/PIENetworkComponent.h"
+
+#include "GameFramework/GameModeBase.h"
+
+#if ENABLE_PIE_NETWORK_TEST
+
+NETWORK_TEST_CLASS(MyFixtureName, "Network.Example")
+{
+	// Deriving from the FBasePIENetworkComponentState to add objects to be tested in each state
+	struct DerivedState : public FBasePIENetworkComponentState
+	{
+		APawn* ReplicatedPawn;
+	};
+
+	FPIENetworkComponent<DerivedState> Network{ TestRunner, TestCommandBuilder, bInitializing };
+
+	BEFORE_EACH()
+	{
+		FNetworkComponentBuilder<DerivedState>()
+			.WithClients(2)
+			.WithGameInstanceClass(UGameInstance::StaticClass())
+			.WithGameMode(AGameModeBase::StaticClass())
+			.Build(Network);
+	}
+
+	TEST_METHOD(SpawnAndReplicatePawn_WithReplicatedPawn_ProvidesPawnToClients)
+	{
+		Network.SpawnAndReplicate<APawn, &DerivedState::ReplicatedPawn>()
+			.ThenServer([this](DerivedState& ServerState) {
+				ASSERT_THAT(IsNotNull(ServerState.ReplicatedPawn));
+			})
+			.ThenClients([this](DerivedState& ClientState) {
+				ASSERT_THAT(IsNotNull(ClientState.ReplicatedPawn));
+			});
+	}
+};
+
+#endif // ENABLE_PIE_NETWORK_TEST
+*/
+
+/**
+ * Struct which handles the PIE session's world and network information.
+ * Both the server and client will keep a reference to their own respective state which will be tested against.
+ */
 struct FBasePIENetworkComponentState
 {
+	/** Reference to the session's world. */
 	UWorld* World = nullptr;
+
+	/** Used by the server to reference the connections to the client. */
 	TArray<UNetConnection*> ClientConnections;
+
+	/** Used by the client to reference location within the `ClientConnections` array on the server. */
 	int32 ClientIndex = INDEX_NONE;
+
+	/** Used by the server for creating and validation of client instances. */
 	int32 ClientCount = 2;
+
+	/** Used by the server to create a PIE session as a dedicated or listen server. */
 	bool bIsDedicatedServer = true;
 };
 
+/**
+ * Component which creates and initializes the server and client network connections between the PIE sessions.
+ * The component also provides methods around the CommandBuilder for adding latent commands.
+ */
 class CQTEST_API FBasePIENetworkComponent
 {
 public:
+	/**
+	 * Construct the BasePIENetworkComponent.
+	 *
+	 * @param InTestRunner - Pointer to the TestRunner used for test reporting.
+	 * @param InCommandBuilder - Reference to the latent command manager.
+	 * @param IsInitializing - Bool specifying if the Automation framework is still being initialized.
+	 * 
+	 * @note requires that the `FBasePIENetworkComponent` is built using the `FNetworkComponentBuilder` as this will setup the server and client `FBasePIENetworkComponentState`
+	 */
 	FBasePIENetworkComponent(FAutomationTestBase* InTestRunner, FTestCommandBuilder& InCommandBuilder, bool IsInitializing);
 
+	/**
+	 * Enqueues a function to the latent command manager for execution on a subsequent frame.
+	 * 
+	 * @param Action - Function to be executed.
+	 * 
+	 * @return a reference to this
+	 */
 	FBasePIENetworkComponent& Then(TFunction<void()> Action);
+
+	/**
+	 * Enqueues a function to the latent command manager for execution on a subsequent frame.
+	 *
+	 * @param Description - Description of the provided function to be logged.
+	 * @param Action - Function to be executed.
+	 * 
+	 * @return a reference to this
+	 */
 	FBasePIENetworkComponent& Then(const TCHAR* Description, TFunction<void()> Action);
 
+	/**
+	 * Enqueues a function to the latent command manager for execution on a subsequent frame.
+	 *
+	 * @param Action - Function to be executed.
+	 * 
+	 * @return a reference to this
+	 */
 	FBasePIENetworkComponent& Do(TFunction<void()> Action);
+
+	/**
+	 * Enqueues a function to the latent command manager for execution on a subsequent frame.
+	 *
+	 * @param Description - Description of the provided function to be logged.
+	 * @param Action - Function to be executed.
+	 * 
+	 * @return a reference to this
+	 */
 	FBasePIENetworkComponent& Do(const TCHAR* Description, TFunction<void()> Action);
 
+	/**
+	 * Enqueues a function to the latent command manager for execution until completion or timeout.
+	 *
+	 * @param Query - Function to be executed.
+	 * @param Timeout - Duration that the latent command can execute before reporting an error.
+	 * 
+	 * @return a reference to this
+	 * 
+	 * @note Will continue to execute until the function provided evaluates to `true` or the duration exceeds the specified timeout.
+	 */
 	FBasePIENetworkComponent& Until(TFunction<bool()> Query, FTimespan Timeout = FTimespan::FromSeconds(10));
+
+	/**
+	 * Enqueues a function to the latent command manager for execution until completion or timeout.
+	 *
+	 * @param Description - Description of the provided function to be logged.
+	 * @param Query - Function to be executed.
+	 * @param Timeout - Duration that the latent command can execute before reporting an error.
+	 * 
+	 * @return a reference to this
+	 *
+	 * @note Will continue to execute until the function provided evaluates to `true` or the duration exceeds the specified timeout.
+	 */
 	FBasePIENetworkComponent& Until(TCHAR* Description, TFunction<bool()> Query, FTimespan Timeout = FTimespan::FromSeconds(10));
 	
+	/**
+	 * Enqueues a function to the latent command manager for execution until completion or timeout.
+	 *
+	 * @param Query - Function to be executed.
+	 * @param Timeout - Duration that the latent command can execute before reporting an error.
+	 * 
+	 * @return a reference to this
+	 *
+	 * @note Will continue to execute until the function provided evaluates to `true` or the duration exceeds the specified timeout.
+	 */
 	FBasePIENetworkComponent& StartWhen(TFunction<bool()> Query, FTimespan Timeout = FTimespan::FromSeconds(10));
+
+	/**
+	 * Enqueues a function to the latent command manager for execution until completion or timeout.
+	 *
+	 * @param Description - Description of the provided function to be logged.
+	 * @param Query - Function to be executed.
+	 * @param Timeout - Duration that the latent command can execute before reporting an error.
+	 * 
+	 * @return a reference to this
+	 *
+	 * @note Will continue to execute until the function provided evaluates to `true` or the duration exceeds the specified timeout.
+	 */
 	FBasePIENetworkComponent& StartWhen(TCHAR* Description, TFunction<bool()> Query, FTimespan Timeout = FTimespan::FromSeconds(10));
 
 protected:
+	/** Prepares the active PIE sessions to be stopped. */
 	void StopPie();
+
+	/** Setup settings for a network session and start PIE sessions for both the server and client with the network settings applied. */
 	void StartPie();
+
+	/**
+	 * Fetch all of the worlds that will be used for the networked session.
+	 *
+	 * @return true if an error was encountered or if the PIE sessions were created with the correct network settings, false otherwise.
+	 *
+	 * @note Method is expected to be used within the `Until` latent command to then wait until the worlds are ready for use.
+	 */
 	bool CollectPieWorlds();
+
+	/**
+	 * Go through all of the client connections to make sure they are connected and ready.
+	 *
+	 * @return true if an error was encountered or if the connections all have a valid controller, false otherwise.
+	 *
+	 * @note Method is expected to be used within the `Until` latent command to then wait until the worlds are ready for use.
+	 */
 	bool AwaitConnections();
+
+	/** Restore back to the state prior to test execution. */
 	void RestoreState();
 
 	TUniquePtr<FBasePIENetworkComponentState> ServerState = nullptr;
@@ -58,31 +225,155 @@ protected:
 template <typename NetworkDataType>
 class FNetworkComponentBuilder;
 
+/**
+ * Component which creates and initializes the server and client network connections between the PIE sessions.
+ * Expands on the `FBasePIENetworkComponent` by providing separate methods to add latent commands for the server and clients.
+ */
 template <typename NetworkDataType = FBasePIENetworkComponentState>
 class FPIENetworkComponent : public FBasePIENetworkComponent
 {
 public:
+	/**
+	 * Construct the PIENetworkComponent.
+	 *
+	 * @param InTestRunner - Pointer to the TestRunner used for test reporting.
+	 * @param InCommandBuilder - Reference to the latent command manager.
+	 * @param IsInitializing - Bool specifying if the Automation framework is still being initialized.
+	 *
+	 * @note Requires that the `FPIENetworkComponent` is built using the `FNetworkComponentBuilder` as this will setup the server and client `NetworkDataType` states
+	 */
 	FPIENetworkComponent(FAutomationTestBase* InTestRunner, FTestCommandBuilder& InCommandBuilder, bool IsInitializing)
 		: FBasePIENetworkComponent(InTestRunner, InCommandBuilder, IsInitializing) {}
 
+	/**
+	 * Enqueues a function to the latent command manager for execution on a subsequent frame on the server.
+	 *
+	 * @param Action - Function to be executed.
+	 * 
+	 * @return a reference to this
+	 */
 	FPIENetworkComponent& ThenServer(TFunction<void(NetworkDataType&)> Action);
+
+	/**
+	 * Enqueues a function to the latent command manager for execution on a subsequent frame on the server.
+	 *
+	 * @param Description - Description of the provided function to be logged.
+	 * @param Action - Function to be executed.
+	 * 
+	 * @return a reference to this
+	 */
 	FPIENetworkComponent& ThenServer(const TCHAR* Description, TFunction<void(NetworkDataType&)> Action);
 
+	/**
+	 * Enqueues a function to the latent command manager for execution on a subsequent frame on all clients.
+	 *
+	 * @param Action - Function to be executed.
+	 * 
+	 * @return a reference to this
+	 */
 	FPIENetworkComponent& ThenClients(TFunction<void(NetworkDataType&)> Action);
+
+	/**
+	 * Enqueues a function to the latent command manager for execution on a subsequent frame on all clients.
+	 *
+	 * @param Description - Description of the provided function to be logged.
+	 * @param Action - Function to be executed.
+	 * 
+	 * @return a reference to this
+	 */
 	FPIENetworkComponent& ThenClients(const TCHAR* Description, TFunction<void(NetworkDataType&)> Action);
 
+	/**
+	 * Enqueues a function to the latent command manager for execution on a subsequent frame on a given client.
+	 *
+	 * @param ClientIndex - Index of the connected client.
+	 * @param Action - Function to be executed.
+	 * 
+	 * @return a reference to this
+	 */
 	FPIENetworkComponent& ThenClient(int32 ClientIndex, TFunction<void(NetworkDataType&)> Action);
+
+	/**
+	 * Enqueues a function to the latent command manager for execution on a subsequent frame on a given client.
+	 *
+	 * @param Description - Description of the provided function to be logged.
+	 * @param ClientIndex - Index of the connected client.
+	 * @param Action - Function to be executed.
+	 * 
+	 * @return a reference to this
+	 */
 	FPIENetworkComponent& ThenClient(const TCHAR* Description, int32 ClientIndex, TFunction<void(NetworkDataType&)> Action);
 
+	/**
+	 * Enqueues a function to the latent command manager for execution until completion or timeout on the server.
+	 *
+	 * @param Query - Function to be executed.
+	 * @param Timeout - Duration that the latent command can execute before reporting an error.
+	 * 
+	 * @return a reference to this
+	 */
 	FPIENetworkComponent& UntilServer(TFunction<bool(NetworkDataType&)> Query, FTimespan Timeout = FTimespan::FromSeconds(10));
+
+	/**
+	 * Enqueues a function to the latent command manager for execution until completion or timeout on the server.
+	 *
+	 * @param Description - Description of the latent command
+	 * @param Query - Function to be executed.
+	 * @param Timeout - Duration that the latent command can execute before reporting an error.
+	 * 
+	 * @return a reference to this
+	 */
 	FPIENetworkComponent& UntilServer(const TCHAR* Description, TFunction<bool(NetworkDataType&)> Query, FTimespan Timeout = FTimespan::FromSeconds(10));
 
+	/**
+	 * Enqueues a function to the latent command manager for execution until completion or timeout on all clients.
+	 *
+	 * @param Query - Function to be executed.
+	 * @param Timeout - Duration that the latent command can execute before reporting an error.
+	 * 
+	 * @return a reference to this
+	 */
 	FPIENetworkComponent& UntilClients(TFunction<bool(NetworkDataType&)> Query, FTimespan Timeout = FTimespan::FromSeconds(10));
+
+	/**
+	 * Enqueues a function to the latent command manager for execution until completion or timeout on all clients.
+	 *
+	 * @param Description - Description of the latent command
+	 * @param Query - Function to be executed.
+	 * @param Timeout - Duration that the latent command can execute before reporting an error.
+	 * 
+	 * @return a reference to this
+	 */
 	FPIENetworkComponent& UntilClients(const TCHAR* Description, TFunction<bool(NetworkDataType&)> Query, FTimespan Timeout = FTimespan::FromSeconds(10));
 	
+	/**
+	 * Enqueues a function to the latent command manager for execution until completion or timeout on a given client.
+	 *
+	 * @param ClientIndex - Index of the connected client.
+	 * @param Query - Function to be executed.
+	 * @param Timeout - Duration that the latent command can execute before reporting an error.
+	 * 
+	 * @return a reference to this
+	 */
 	FPIENetworkComponent& UntilClient(int32 ClientIndex, TFunction<bool(NetworkDataType&)> Query, FTimespan Timeout = FTimespan::FromSeconds(10));
+
+	/**
+	 * Enqueues a function to the latent command manager for execution until completion or timeout on a given client.
+	 *
+	 * @param Description - Description of the latent command
+	 * @param ClientIndex - Index of the connected client.
+	 * @param Query - Function to be executed.
+	 * @param Timeout - Duration that the latent command can execute before reporting an error.
+	 * 
+	 * @return a reference to this
+	 */
 	FPIENetworkComponent& UntilClient(const TCHAR* Description, int32 ClientIndex, TFunction<bool(NetworkDataType&)> Query, FTimespan Timeout = FTimespan::FromSeconds(10));
 
+	/**
+	 * Spawns an actor on the server and replicates to the clients.
+	 *
+	 * @return a reference to this
+	 */
 	template<typename ActorToSpawn, ActorToSpawn* NetworkDataType::*ResultStorage>
 	FPIENetworkComponent& SpawnAndReplicate();
 
@@ -91,19 +382,79 @@ private:
 };
 
 
+/**
+ * Helper object used to setup and build the `FPIENetworkComponent`.
+ */
 template <typename NetworkDataType = FBasePIENetworkComponentState>
 class FNetworkComponentBuilder
 {
 public:
+	/**
+	 * Construct the NetworkComponentBuilder.
+	 *
+	 * @note Will assert if the `NetworkDataType` used is not of type or derived from `FBasePIENetworkComponentState`
+	 */
 	FNetworkComponentBuilder();
 
+	/**
+	 * Specifies the number of clients to be used.
+	 *
+	 * @param ClientCount - Number of clients not including the listen server.
+	 * 
+	 * @return a reference to this
+	 * 
+	 * @note Number of clients should exclude the server, even if using the server as a listen server.
+	 */
 	FNetworkComponentBuilder& WithClients(int32 ClientCount);
+
+	/**
+	 * Build the server as a dedicated server.
+	 * 
+	 * @return a reference to this
+	 */
 	FNetworkComponentBuilder& AsDedicatedServer();
+
+	/**
+	 * Build the server as a listen server.
+	 *
+	 * @return a reference to this
+	 */
 	FNetworkComponentBuilder& AsListenServer();
-	FNetworkComponentBuilder& WithPacketSimulationSettings(FPacketSimulationSettings* InPacketSImulationSettings);
+
+	/**
+	 * Build the FPIENetworkComponent to use the provided packet simulation settings.
+	 *
+	 * @param InPacketSimulationSettings - Settings used to simulate packet behavior over a network.
+	 * 
+	 * @return a reference to this
+	 */
+	FNetworkComponentBuilder& WithPacketSimulationSettings(FPacketSimulationSettings* InPacketSimulationSettings);
+
+	/**
+	 * Build the FPIENetworkComponent to use the provided game mode.
+	 *
+	 * @param InGameMode - GameMode used for the networked session.
+	 * 
+	 * @return a reference to this
+	 */
 	FNetworkComponentBuilder& WithGameMode(TSubclassOf<AGameModeBase> InGameMode);
+
+	/**
+	 * Build the FPIENetworkComponent to restore back to the provided game instance.
+	 *
+	 * @param InGameInstanceClass - GameInstance to restore back to.
+	 * 
+	 * @return a reference to this
+	 */
 	FNetworkComponentBuilder& WithGameInstanceClass(FSoftClassPath InGameInstanceClass);
 
+	/**
+	 * Build the FPIENetworkComponent with the provided data.
+	 *
+	 * @param OutNetwork - Reference to the FPIENetworkComponent.
+	 * 
+	 * @return a reference to this
+	 */
 	void Build(FPIENetworkComponent<NetworkDataType>& OutNetwork);
 
 private:
@@ -116,6 +467,7 @@ private:
 
 constexpr static EAutomationTestFlags NetworkTestContext = EAutomationTestFlags::EditorContext | EAutomationTestFlags::ProductFilter;
 
+/** Macro to quickly create tests which will only run within the Editor. */
 #define NETWORK_TEST_CLASS(_ClassName, _TestDir) TEST_CLASS_WITH_FLAGS(_ClassName, _TestDir, NetworkTestContext)
 
 #include "Impl/PIENetworkComponent.inl"
diff --git a/Engine/Plugins/Tests/CQTest/Source/CQTest/Public/Components/PIENetworkTestStateRestorer.h b/Engine/Plugins/Tests/CQTest/Source/CQTest/Public/Components/PIENetworkTestStateRestorer.h
--- a/Engine/Plugins/Tests/CQTest/Source/CQTest/Public/Components/PIENetworkTestStateRestorer.h
+++ b/Engine/Plugins/Tests/CQTest/Source/CQTest/Public/Components/PIENetworkTestStateRestorer.h
@@ -7,14 +7,24 @@
 class UWorld;
 class AGameModeBase;
 
-//Object which saves the state of the game before starting the test,
-//Then restores it after the test is complete.
+/**
+ * Object which saves the state of the game before starting the test,
+ * Then restores it after the test is complete.
+ */
 class CQTEST_API FPIENetworkTestStateRestorer
 {
 public:
 	FPIENetworkTestStateRestorer() = default;
+
+	/**
+	 * Construct the PIENetworkTestStateRestorer.
+	 *
+	 * @param InGameInstanceClass - Current game instance prior to the test.
+	 * @param InGameMode - Current game mode prior to the test.
+	 */
 	FPIENetworkTestStateRestorer(const FSoftClassPath InGameInstanceClass, TSubclassOf<AGameModeBase> InGameMode);
 
+	/** Restore to the previous game instance and mode. */
 	void Restore();
 
 private:
diff --git a/Engine/Plugins/Tests/CQTest/Source/CQTest/Public/Components/SpawnHelper.h b/Engine/Plugins/Tests/CQTest/Source/CQTest/Public/Components/SpawnHelper.h
--- a/Engine/Plugins/Tests/CQTest/Source/CQTest/Public/Components/SpawnHelper.h
+++ b/Engine/Plugins/Tests/CQTest/Source/CQTest/Public/Components/SpawnHelper.h
@@ -6,22 +6,55 @@
 #include "GameFramework/Actor.h"
 #include "UObject/GCObject.h"
 
+/**
+ * Helper object for spawning Actors and other object types in the world.
+ * 
+ * @see FActorTestSpawner, FMapTestSpawner
+ */
 struct CQTEST_API FSpawnHelper
 {
+	/** Destruct the Spawn Helper. */
 	virtual ~FSpawnHelper();
 
+	/**
+	 * Spawn an Actor in the world.
+	 *
+	 * @param SpawnParameters - Struct of optional parameters used to assist with spawning.
+	 * @param Class - Class of the object to be spawned.
+	 * 
+	 * @return reference to the spawned object
+	 * 
+	 * @note Method guarantees that the Actor returned is valid, will assert otherwise.
+	 */
 	template <typename ActorType>
 	ActorType& SpawnActor(const FActorSpawnParameters& SpawnParameters = FActorSpawnParameters(), UClass* Class = nullptr)
 	{
 		return SpawnActorAtInWorld<ActorType>(GetWorld(), FVector::ZeroVector, FRotator::ZeroRotator, SpawnParameters, Class);
 	}
 
+	/**
+	 * Spawn an Actor in the world.
+	 *
+	 * @param Location - Location to spawn the Actor in the world.
+	 * @param Rotation - Rotation to spawn the Actor in the world.
+	 * @param SpawnParameters - Struct of optional parameters used to assist with spawning.
+	 * @param Class - Class of the object to be spawned.
+	 * 
+	 * @return reference to the spawned object
+	 *
+	 * @note Method guarantees that the Actor returned is valid, will assert otherwise.
+	 */
 	template <typename ActorType>
 	ActorType& SpawnActorAt(FVector const& Location, FRotator const& Rotation, const FActorSpawnParameters& SpawnParameters = FActorSpawnParameters(), UClass* Class = nullptr)
 	{
 		return SpawnActorAtInWorld<ActorType>(GetWorld(), Location, Rotation, SpawnParameters, Class);
 	}
 
+	/**
+	 * Create a new Object.
+	 *
+	 * @return reference to the created object
+	 */
 	template <typename ObjectType>
 	ObjectType& SpawnObject()
 	{
@@ -33,9 +66,15 @@ struct CQTEST_API FSpawnHelper
 		return *Object;
 	}
 
+	/** Returns a reference to the current world. */
 	UWorld& GetWorld();
 
 protected:
+	/**
+	 * Creates a new world.
+	 * 
+	 * @returns a newly created world.
+	 */
 	virtual UWorld* CreateWorld() = 0;
 
 	TArray<TWeakObjectPtr<AActor>> SpawnedActors{};
@@ -44,6 +83,19 @@ protected:
 	UWorld* GameWorld{ nullptr };
 
 private:
+	/**
+	 * Spawn an Actor in the world.
+	 *
+	 * @param World - World to spawn the Actor.
+	 * @param Location - Location to spawn the Actor in the world.
+	 * @param Rotation - Rotation to spawn the Actor in the world.
+	 * @param SpawnParameters - Struct of optional parameters used to assist with spawning.
+	 * @param Class - Class of the object to be spawned.
+	 * 
+	 * @return reference to the spawned object
+	 *
+	 * @note Method guarantees that the Actor returned is valid, will assert otherwise.
+	 */
 	template <typename ActorType>
 	ActorType& SpawnActorAtInWorld(UWorld& World, const FVector& Location, const FRotator& Rotation, const FActorSpawnParameters& SpawnParameters = FActorSpawnParameters(), UClass* Class = nullptr)
 	{
diff --git a/Engine/Plugins/Tests/CQTest/Source/CQTest/Public/Impl/PIENetworkComponent.inl b/Engine/Plugins/Tests/CQTest/Source/CQTest/Public/Impl/PIENetworkComponent.inl
--- a/Engine/Plugins/Tests/CQTest/Source/CQTest/Public/Impl/PIENetworkComponent.inl
+++ b/Engine/Plugins/Tests/CQTest/Source/CQTest/Public/Impl/PIENetworkComponent.inl
@@ -1,5 +1,8 @@
 // Copyright Epic Games, Inc. All Rights Reserved.
 
+///////////////////////////////////////////////////////////////////////
+// FPIENetworkComponent
+
 template <typename NetworkDataType>
 inline FPIENetworkComponent<NetworkDataType>& FPIENetworkComponent<NetworkDataType>::ThenServer(TFunction<void(NetworkDataType&)> Action)
 {
@@ -153,7 +156,8 @@ inline FPIENetworkComponent<NetworkDataType>& FPIENetworkComponent<NetworkDataTy
 	return *this;
 }
 
-//////////////////////
+///////////////////////////////////////////////////////////////////////
+// FNetworkComponentBuilder
 
 template <typename NetworkDataType>
 inline FNetworkComponentBuilder<NetworkDataType>::FNetworkComponentBuilder()
diff --git a/Engine/Plugins/Tests/CQTest/Source/CQTestTests/Private/Components/SpawnHelperTests.cpp b/Engine/Plugins/Tests/CQTest/Source/CQTestTests/Private/Components/SpawnHelperTests.cpp
--- a/Engine/Plugins/Tests/CQTest/Source/CQTestTests/Private/Components/SpawnHelperTests.cpp
+++ b/Engine/Plugins/Tests/CQTest/Source/CQTestTests/Private/Components/SpawnHelperTests.cpp
@@ -28,7 +28,7 @@ TEST_CLASS(SpawnHelperTests, "TestFramework.CQTest.Actor")
 		ASSERT_THAT(AreEqual(TestComponent->GetOwner(), &Owner ));
 	}
 
-	TEST_METHOD(SpawnObject_WithAnbject_TearsDownCleanly)
+	TEST_METHOD(SpawnObject_WithAnObject_TearsDownCleanly)
 	{
 		auto& SomeComponent = Spawner.SpawnObject<USceneComponent>();
 	}
